UE Jpeg: name change
#p4v-preflight-copy 29191952

diff --git a/Engine/Config/BaseEditor.ini b/Engine/Config/BaseEditor.ini
--- a/Engine/Config/BaseEditor.ini
+++ b/Engine/Config/BaseEditor.ini
@@ -261,7 +261,7 @@ EnableOodleWAV=True
 AllowNonPowerOfTwoTextures=True
 ; Imported jpegs will be stored in packages in jpeg format until edited, this will help keep texture asset sizes down
 RetainJpegFormat=True
-EnableOodleJPEG=False
+EnableUEJpeg=False
 
 [/Script/Cascade.CascadeConfiguration]
 +ModuleMenu_ModuleRejections=ParticleModuleTypeDataBase
diff --git a/Engine/Plugins/Interchange/Runtime/Source/Import/Private/InterchangeImportModule.cpp b/Engine/Plugins/Interchange/Runtime/Source/Import/Private/InterchangeImportModule.cpp
--- a/Engine/Plugins/Interchange/Runtime/Source/Import/Private/InterchangeImportModule.cpp
+++ b/Engine/Plugins/Interchange/Runtime/Source/Import/Private/InterchangeImportModule.cpp
@@ -32,7 +32,7 @@
 #include "Texture/InterchangeDDSTranslator.h"
 #include "Texture/InterchangeIESTranslator.h"
 #include "Texture/InterchangeJPGTranslator.h"
-#include "Texture/InterchangeOOJPGTranslator.h"
+#include "Texture/InterchangeUEJPEGTranslator.h"
 #include "Texture/InterchangePCXTranslator.h"
 #include "Texture/InterchangePSDTranslator.h"
 #include "Texture/InterchangeTextureFactory.h"
@@ -87,7 +87,7 @@ void FInterchangeImportModule::StartupModule()
 		//Textures
 		InterchangeManager.RegisterTranslator(UInterchangeImageWrapperTranslator::StaticClass());
 		InterchangeManager.RegisterTranslator(UInterchangeDDSTranslator::StaticClass());
-		InterchangeManager.RegisterTranslator(UInterchangeOOJPGTranslator::StaticClass());
+		InterchangeManager.RegisterTranslator(UInterchangeUEJPEGTranslator::StaticClass());
 		InterchangeManager.RegisterTranslator(UInterchangeJPGTranslator::StaticClass());
 		InterchangeManager.RegisterTranslator(UInterchangePCXTranslator::StaticClass());
 		InterchangeManager.RegisterTranslator(UInterchangePSDTranslator::StaticClass());
diff --git a/Engine/Plugins/Interchange/Runtime/Source/Import/Private/Texture/InterchangeOOJPGTranslator.cpp b/Engine/Plugins/Interchange/Runtime/Source/Import/Private/Texture/InterchangeOOJPGTranslator.cpp
deleted file mode 100644
diff --git a/Engine/Plugins/Interchange/Runtime/Source/Import/Private/Texture/InterchangeTextureFactory.cpp b/Engine/Plugins/Interchange/Runtime/Source/Import/Private/Texture/InterchangeTextureFactory.cpp
--- a/Engine/Plugins/Interchange/Runtime/Source/Import/Private/Texture/InterchangeTextureFactory.cpp
+++ b/Engine/Plugins/Interchange/Runtime/Source/Import/Private/Texture/InterchangeTextureFactory.cpp
@@ -41,7 +41,7 @@
 #include "Serialization/EditorBulkData.h"
 #include "Texture/InterchangeBlockedTexturePayloadInterface.h"
 #include "Texture/InterchangeJPGTranslator.h"
-#include "Texture/InterchangeOOJPGTranslator.h"
+#include "Texture/InterchangeUEJPEGTranslator.h"
 #include "Texture/InterchangeSlicedTexturePayloadInterface.h"
 #include "Texture/InterchangeTextureLightProfilePayloadInterface.h"
 #include "Texture/InterchangeTexturePayloadInterface.h"
@@ -680,7 +680,7 @@ namespace UE::Interchange::Private::InterchangeTextureFactory
 				if (BlockAndSourceDataFiles.IsEmpty())
 				{
 					if (Translator->GetClass() == UInterchangeJPGTranslator::StaticClass() 
-						|| Translator->GetClass() == UInterchangeOOJPGTranslator::StaticClass())
+						|| Translator->GetClass() == UInterchangeUEJPEGTranslator::StaticClass())
 					{
 						// Honor setting from TextureImporter.RetainJpegFormat in Editor.ini if it exists (ideally we should deprecate this as it is confusing and probably not thread safe)
 						const bool bShouldImportRawCache = bShoudImportCompressedImage;
diff --git a/Engine/Plugins/Interchange/Runtime/Source/Import/Private/Texture/InterchangeUEJPEGTranslator.cpp b/Engine/Plugins/Interchange/Runtime/Source/Import/Private/Texture/InterchangeUEJPEGTranslator.cpp
new file mode 100644
--- /dev/null
+++ b/Engine/Plugins/Interchange/Runtime/Source/Import/Private/Texture/InterchangeUEJPEGTranslator.cpp
@@ -0,0 +1,147 @@
+// Copyright Epic Games, Inc. All Rights Reserved. 
+#include "Texture/InterchangeUEJPEGTranslator.h"
+
+#include "Algo/Find.h"
+#include "Containers/StaticArray.h"
+#include "Engine/Texture.h"
+#include "Engine/Texture2D.h"
+#include "IImageWrapper.h"
+#include "IImageWrapperModule.h"
+#include "ImageCoreUtils.h"
+#include "InterchangeImportLog.h"
+#include "InterchangeTextureNode.h"
+#include "Memory/SharedBuffer.h"
+#include "Misc/ConfigCacheIni.h"
+#include "Misc/FileHelper.h"
+#include "Misc/Paths.h"
+#include "Misc/ScopedSlowTask.h"
+#include "Modules/ModuleManager.h"
+#include "Nodes/InterchangeBaseNodeContainer.h"
+#include "Texture/TextureTranslatorUtilities.h"
+
+#include UE_INLINE_GENERATED_CPP_BY_NAME(InterchangeUEJPEGTranslator)
+
+static bool GInterchangeEnableUEJPEGImport = true;
+static FAutoConsoleVariableRef CCvarInterchangeEnableUEJPEGImport(
+	TEXT("Interchange.FeatureFlags.Import.UEJPEG"),
+	GInterchangeEnableUEJPEGImport,
+	TEXT("Whether UEJPEG support is enabled."),
+	ECVF_Default);
+
+TArray<FString> UInterchangeUEJPEGTranslator::GetSupportedFormats() const
+{
+	TArray<FString> Formats;
+
+	if (GInterchangeEnableUEJPEGImport || GIsAutomationTesting)
+	{
+		Formats.Reserve(3);
+
+		Formats.Add(TEXT("uej;UEJpeg image"));
+		bool bEnableUEJpeg = false;
+		GConfig->GetBool(TEXT("TextureImporter"), TEXT("EnableUEJpeg"), bEnableUEJpeg, GEditorIni);
+		if (bEnableUEJpeg)
+		{
+			Formats.Add(TEXT("jpg;JPEG image"));
+			Formats.Add(TEXT("jpeg;JPEG image"));
+		}
+	}
+
+	return Formats;
+}
+
+bool UInterchangeUEJPEGTranslator::Translate(UInterchangeBaseNodeContainer& BaseNodeContainer) const
+{
+	return UE::Interchange::FTextureTranslatorUtilities::Generic2DTextureTranslate(GetSourceData(), BaseNodeContainer);
+}
+
+TOptional<UE::Interchange::FImportImage> UInterchangeUEJPEGTranslator::GetTexturePayloadData(const FString& /*PayloadKey*/, TOptional<FString>& AlternateTexturePath) const
+{
+	using namespace UE::Interchange;
+
+	TArray64<uint8> SourceDataBuffer;
+	if (!FTextureTranslatorUtilities::LoadSourceBuffer(*this, TEXT("UEJPEG"), SourceDataBuffer))
+	{
+		return {};
+	}
+
+	const bool bImportRaw = false;
+	return GetTexturePayloadImplementation(MoveTemp(SourceDataBuffer), bImportRaw);
+
+}
+
+bool UInterchangeUEJPEGTranslator::SupportCompressedTexturePayloadData() const
+{
+	return true;
+}
+
+TOptional<UE::Interchange::FImportImage> UInterchangeUEJPEGTranslator::GetCompressedTexturePayloadData(const FString& /*PayloadKey*/, TOptional<FString>& /*AlternateTexturePath*/) const
+{
+	using namespace UE::Interchange;
+
+	TArray64<uint8> SourceDataBuffer;
+	if (!FTextureTranslatorUtilities::LoadSourceBuffer(*this, TEXT("UEJPEG"), SourceDataBuffer))
+	{
+		return {};
+	}
+
+	const bool bImportRaw = true;
+	return GetTexturePayloadImplementation(MoveTemp(SourceDataBuffer), bImportRaw);
+}
+
+TOptional<UE::Interchange::FImportImage> UInterchangeUEJPEGTranslator::GetTexturePayloadImplementation(TArray64<uint8>&& SourceDataBuffer, bool bShouldImportRaw) const
+{
+	using namespace UE::Interchange;
+
+	IImageWrapperModule& ImageWrapperModule = FModuleManager::LoadModuleChecked<IImageWrapperModule>(FName("ImageWrapper"));
+
+	//
+	// UEJPEG
+	//
+	ETextureSourceCompressionFormat TscfFormat = ETextureSourceCompressionFormat::TSCF_UEJPEG; 
+	TSharedPtr<IImageWrapper> ImageWrapper = ImageWrapperModule.CreateImageWrapper(EImageFormat::UEJPEG);
+	if (!ImageWrapper.IsValid() || !ImageWrapper->SetCompressed(SourceDataBuffer.GetData(), SourceDataBuffer.Num()))
+	{
+		TscfFormat = ETextureSourceCompressionFormat::TSCF_JPEG; 
+		ImageWrapper = ImageWrapperModule.CreateImageWrapper(EImageFormat::JPEG);
+		if (!ImageWrapper.IsValid() || !ImageWrapper->SetCompressed(SourceDataBuffer.GetData(), SourceDataBuffer.Num()))
+		{
+			FTextureTranslatorUtilities::LogError(*this, NSLOCTEXT("InterchangeUEJPEGTranslator", "DecodingFailed", "Failed to decode UEJPEG."));
+			return TOptional<UE::Interchange::FImportImage>();
+		}
+	}
+
+	// Select the texture's source format
+	ERawImageFormat::Type RawFormat = ImageWrapper->GetClosestRawImageFormat();
+	check( RawFormat != ERawImageFormat::Invalid );
+	ETextureSourceFormat TextureFormat = FImageCoreUtils::ConvertToTextureSourceFormat(RawFormat);
+
+	UE::Interchange::FImportImage PayloadData;
+
+	const bool bShouldAllocateRawDataBuffer = false;
+
+	PayloadData.Init2DWithParams(
+		ImageWrapper->GetWidth(),
+		ImageWrapper->GetHeight(),
+		TextureFormat,
+		ImageWrapper->GetSRGB(),
+		bShouldAllocateRawDataBuffer
+	);
+
+	TArray64<uint8> RawData;
+	if (bShouldImportRaw)
+	{
+		PayloadData.RawData = MakeUniqueBufferFromArray(ImageWrapper->GetCompressed());
+		PayloadData.RawDataCompressionFormat = TscfFormat;
+	}
+	else if (ImageWrapper->GetRaw(RawData))
+	{
+		PayloadData.RawData = MakeUniqueBufferFromArray(MoveTemp(RawData));
+	}
+	else
+	{
+		FTextureTranslatorUtilities::LogError(*this, NSLOCTEXT("InterchangeUEJPEGTranslator", "DecodingFailed", "Failed to decode UEJPEG."));
+		return TOptional<UE::Interchange::FImportImage>();
+	}
+
+	return PayloadData;
+}
diff --git a/Engine/Plugins/Interchange/Runtime/Source/Import/Public/Texture/InterchangeOOJPGTranslator.h b/Engine/Plugins/Interchange/Runtime/Source/Import/Public/Texture/InterchangeOOJPGTranslator.h
deleted file mode 100644
diff --git a/Engine/Plugins/Interchange/Runtime/Source/Import/Public/Texture/InterchangeUEJPEGTranslator.h b/Engine/Plugins/Interchange/Runtime/Source/Import/Public/Texture/InterchangeUEJPEGTranslator.h
new file mode 100644
--- /dev/null
+++ b/Engine/Plugins/Interchange/Runtime/Source/Import/Public/Texture/InterchangeUEJPEGTranslator.h
@@ -0,0 +1,63 @@
+// Copyright Epic Games, Inc. All Rights Reserved.
+
+#pragma once
+
+#include "CoreMinimal.h"
+#include "UObject/Object.h"
+#include "UObject/ObjectMacros.h"
+#include "InterchangeTranslatorBase.h"
+#include "Nodes/InterchangeBaseNodeContainer.h"
+#include "Texture/InterchangeTexturePayloadData.h"
+#include "Texture/InterchangeTexturePayloadInterface.h"
+
+#include "InterchangeUEJPEGTranslator.generated.h"
+
+UCLASS(BlueprintType)
+class INTERCHANGEIMPORT_API UInterchangeUEJPEGTranslator : public UInterchangeTranslatorBase, public IInterchangeTexturePayloadInterface
+{
+	GENERATED_BODY()
+public:
+
+	virtual TArray<FString> GetSupportedFormats() const override;
+	virtual EInterchangeTranslatorAssetType GetSupportedAssetTypes() const override { return EInterchangeTranslatorAssetType::Textures; }
+
+	/**
+	 * Translate the associated source data into a node hold by the specified nodes container.
+	 *
+	 * @param BaseNodeContainer - The unreal objects descriptions container where to put the translated source data.
+	 * @return true if the translator can translate the source data, false otherwise.
+	 */
+	virtual bool Translate(UInterchangeBaseNodeContainer& BaseNodeContainer) const override;
+
+	/* IInterchangeTexturePayloadInterface Begin */
+
+	/**
+	 * Once the translation is done, the import process need a way to retrieve payload data.
+	 * This payload will be use by the factories to create the asset.
+	 *
+	 * @param PayloadKey - Unused. The translator uses its SourceData property to extract the payload.
+	 * @param AlternateTexturePath - Unused. The translator uses its SourceData property to extract the payload.
+	 * @return a PayloadData containing the import image data. The TOptional will not be set if there is an error.
+	 */
+	virtual TOptional<UE::Interchange::FImportImage> GetTexturePayloadData(const FString& PayloadKey, TOptional<FString>& AlternateTexturePath) const override;
+
+	virtual bool SupportCompressedTexturePayloadData() const override;
+
+	/**
+	 * Once the translation is done, the import process need a way to retrieve payload data.
+	 * This payload will be use by the factories to create the asset.
+	 *
+	 * @param PayloadKey - Unused. The translator uses its SourceData property to extract the payload.
+	 * @param AlternateTexturePath - Unused. The translator uses its SourceData property to extract the payload.
+	 * @return a PayloadData containing the import image data. The TOptional will not be set if there is an error.
+	 */
+	virtual TOptional<UE::Interchange::FImportImage> GetCompressedTexturePayloadData(const FString& PayloadKey, TOptional<FString>& AlternateTexturePath) const override;
+	
+	/* IInterchangeTexturePayloadInterface End */
+
+private:
+
+	TOptional<UE::Interchange::FImportImage> GetTexturePayloadImplementation(TArray64<uint8>&& SourceDataBuffer, bool bShouldImportRaw) const;
+};
+
+
diff --git a/Engine/Source/Developer/TextureBuild/Private/TextureBuildFunction.cpp b/Engine/Source/Developer/TextureBuild/Private/TextureBuildFunction.cpp
--- a/Engine/Source/Developer/TextureBuild/Private/TextureBuildFunction.cpp
+++ b/Engine/Source/Developer/TextureBuild/Private/TextureBuildFunction.cpp
@@ -221,9 +221,9 @@ static bool TryReadTextureSourceFromCompactBinary(FCbFieldView Source, UE::Deriv
 			ImageWrapper->GetRaw(SourceFormat == TSF_G8 ? ERGBFormat::Gray : ERGBFormat::BGRA, 8, IntermediateDecompressedData);
 		}
 		break;
-		case TSCF_OOJPEG:
+		case TSCF_UEJPEG:
 		{
-			TSharedPtr<IImageWrapper> ImageWrapper = FModuleManager::GetModuleChecked<IImageWrapperModule>(FName("ImageWrapper")).CreateImageWrapper(EImageFormat::OOJPEG);
+			TSharedPtr<IImageWrapper> ImageWrapper = FModuleManager::GetModuleChecked<IImageWrapperModule>(FName("ImageWrapper")).CreateImageWrapper(EImageFormat::UEJPEG);
 			ImageWrapper->SetCompressed((const uint8*)InputBuffer.GetData(), InputBuffer.GetSize());
 			ImageWrapper->GetRaw(SourceFormat == TSF_G8 ? ERGBFormat::Gray : ERGBFormat::BGRA, 8, IntermediateDecompressedData);
 		}
diff --git a/Engine/Source/Editor/UnrealEd/Classes/Exporters/TextureExporterPNG.h b/Engine/Source/Editor/UnrealEd/Classes/Exporters/TextureExporterPNG.h
--- a/Engine/Source/Editor/UnrealEd/Classes/Exporters/TextureExporterPNG.h
+++ b/Engine/Source/Editor/UnrealEd/Classes/Exporters/TextureExporterPNG.h
@@ -44,7 +44,7 @@ class UTextureExporterJPEG : public UExporter
 
 
 UCLASS()
-class UTextureExporterOOJPEG : public UTextureExporterGeneric
+class UTextureExporterUEJPEG : public UTextureExporterGeneric
 {
 	GENERATED_UCLASS_BODY()
 	
diff --git a/Engine/Source/Editor/UnrealEd/Private/Factories/EditorFactories.cpp b/Engine/Source/Editor/UnrealEd/Private/Factories/EditorFactories.cpp
--- a/Engine/Source/Editor/UnrealEd/Private/Factories/EditorFactories.cpp
+++ b/Engine/Source/Editor/UnrealEd/Private/Factories/EditorFactories.cpp
@@ -2932,7 +2932,7 @@ bool UTextureFactory::ImportImage(const uint8* Buffer, int64 Length, FFeedbackCo
 			// branch for JPEG, if retaining the jpeg compressed data
 			// this is inside the DecompressImage branch even though we don't use the LoadedImage at all
 			//	 just to ensure that the jpeg will decode successfully
-			if (ImageFormat == EImageFormat::JPEG || ImageFormat == EImageFormat::OOJPEG)
+			if (ImageFormat == EImageFormat::JPEG || ImageFormat == EImageFormat::UEJPEG)
 			{
 				// unusual loader, retains jpeg
 				bool bRetainJpegFormat = false;
@@ -2971,8 +2971,8 @@ bool UTextureFactory::ImportImage(const uint8* Buffer, int64 Length, FFeedbackCo
 						}
 						else
 						{
-							OutImage.RawDataCompressionFormat = ETextureSourceCompressionFormat::TSCF_OOJPEG;
-							UE_LOG(LogEditorFactories,Display,TEXT("OOJPEG imported and retained as OOJPEG in uasset."));
+							OutImage.RawDataCompressionFormat = ETextureSourceCompressionFormat::TSCF_UEJPEG;
+							UE_LOG(LogEditorFactories,Display,TEXT("UEJPEG imported and retained as UEJPEG in uasset."));
 						}
 
 						return true;
@@ -5004,14 +5004,14 @@ UVirtualTextureBuilderExporterPNG::UVirtualTextureBuilderExporterPNG(const FObje
 
 //========================================================
 
-UTextureExporterOOJPEG::UTextureExporterOOJPEG(const FObjectInitializer& ObjectInitializer)
+UTextureExporterUEJPEG::UTextureExporterUEJPEG(const FObjectInitializer& ObjectInitializer)
 	: UTextureExporterGeneric(ObjectInitializer)
 {
-	FormatExtension.Add(TEXT("OOJ"));
-	FormatDescription.Add(TEXT("Oodle JPEG"));
+	FormatExtension.Add(TEXT("UEJ"));
+	FormatDescription.Add(TEXT("UE JPEG"));
 }
 
-bool UTextureExporterOOJPEG::SupportsTexture(UTexture* Texture) const
+bool UTextureExporterUEJPEG::SupportsTexture(UTexture* Texture) const
 {
 	ETextureSourceFormat TSF = Texture->Source.GetFormat();
 	ERawImageFormat::Type RawFormat = FImageCoreUtils::ConvertToRawImageFormat(TSF);
@@ -5219,7 +5219,7 @@ bool UTextureExporterJPEG::SupportsObject(UObject* Object) const
 			}
 
 			// Check it has JPEG BulkData :
-			if ( (Texture->Source.GetSourceCompression() == TSCF_JPEG || Texture->Source.GetSourceCompression() == TSCF_OOJPEG) &&
+			if ( (Texture->Source.GetSourceCompression() == TSCF_JPEG || Texture->Source.GetSourceCompression() == TSCF_UEJPEG) &&
 				Texture->Source.GetSizeOnDisk() > 0 )
 			{
 				ETextureSourceFormat TSF = Texture->Source.GetFormat();
@@ -5238,7 +5238,7 @@ bool UTextureExporterJPEG::ExportBinary(UObject* Object, const TCHAR* Type, FArc
 	UTexture2D* Texture = Cast<UTexture2D>(Object);
 	check(Texture != nullptr);
 
-	check((Texture->Source.GetSourceCompression() == TSCF_JPEG || Texture->Source.GetSourceCompression() == TSCF_OOJPEG) &&
+	check((Texture->Source.GetSourceCompression() == TSCF_JPEG || Texture->Source.GetSourceCompression() == TSCF_UEJPEG) &&
 			Texture->Source.GetSizeOnDisk() > 0 );
 
 	// just write the JPEG data we already have :
@@ -5253,7 +5253,7 @@ bool UTextureExporterJPEG::ExportBinary(UObject* Object, const TCHAR* Type, FArc
 		else
 		{
 			IImageWrapperModule& ImageWrapperModule = FModuleManager::LoadModuleChecked<IImageWrapperModule>(FName("ImageWrapper"));
-			TSharedPtr<IImageWrapper> ImageWrapper = ImageWrapperModule.CreateImageWrapper(EImageFormat::OOJPEG);
+			TSharedPtr<IImageWrapper> ImageWrapper = ImageWrapperModule.CreateImageWrapper(EImageFormat::UEJPEG);
 			ImageWrapper->SetCompressed(BulkDataBuffer.GetData(), BulkDataBuffer.GetSize());
 			TArray64<uint8> ExportData = ImageWrapper->GetExportData();
 			Ar.Serialize(ExportData.GetData(), ExportData.Num());
diff --git a/Engine/Source/Runtime/Engine/Classes/Engine/TextureDefines.h b/Engine/Source/Runtime/Engine/Classes/Engine/TextureDefines.h
--- a/Engine/Source/Runtime/Engine/Classes/Engine/TextureDefines.h
+++ b/Engine/Source/Runtime/Engine/Classes/Engine/TextureDefines.h
@@ -323,7 +323,7 @@ enum ETextureSourceCompressionFormat : int
 	TSCF_None	UMETA(DisplayName = "None"),
 	TSCF_PNG	UMETA(DisplayName = "PNG"),
 	TSCF_JPEG	UMETA(DisplayName = "JPEG"),
-	TSCF_OOJPEG	UMETA(DisplayName = "Oodle JPEG"),
+	TSCF_UEJPEG	UMETA(DisplayName = "UE JPEG"),
 
 	TSCF_MAX
 };
diff --git a/Engine/Source/Runtime/ImageWrapper/ImageWrapper.Build.cs b/Engine/Source/Runtime/ImageWrapper/ImageWrapper.Build.cs
--- a/Engine/Source/Runtime/ImageWrapper/ImageWrapper.Build.cs
+++ b/Engine/Source/Runtime/ImageWrapper/ImageWrapper.Build.cs
@@ -38,7 +38,7 @@ public class ImageWrapper : ModuleRules
 			"UElibPNG",
 			"LibTiff",
 			"OodleDataCompression",
-			"OodleJPEG"
+			"UEJpegComp"
 		);
 
 		// Jpeg Decoding
diff --git a/Engine/Source/Runtime/ImageWrapper/Private/Formats/OOJpegImageWrapper.cpp b/Engine/Source/Runtime/ImageWrapper/Private/Formats/OOJpegImageWrapper.cpp
deleted file mode 100644
diff --git a/Engine/Source/Runtime/ImageWrapper/Private/Formats/OOJpegImageWrapper.h b/Engine/Source/Runtime/ImageWrapper/Private/Formats/OOJpegImageWrapper.h
deleted file mode 100644
diff --git a/Engine/Source/Runtime/ImageWrapper/Private/Formats/UEJpegImageWrapper.cpp b/Engine/Source/Runtime/ImageWrapper/Private/Formats/UEJpegImageWrapper.cpp
new file mode 100644
--- /dev/null
+++ b/Engine/Source/Runtime/ImageWrapper/Private/Formats/UEJpegImageWrapper.cpp
@@ -0,0 +1,301 @@
+// Copyright Epic Games, Inc. All Rights Reserved.
+
+#include "Formats/UEJpegImageWrapper.h"
+
+#include "Async/ParallelFor.h"
+#include "Math/Color.h"
+#include "Misc/ScopeLock.h"
+#include "ImageWrapperPrivate.h"
+#include "ImageCoreUtils.h"
+
+#if WITH_UEJPEG
+
+#include "uejpeg.h"
+
+static void* ue_malloc(size_t size) 
+{ 
+	return FMemory::Malloc(size); 
+}
+
+static void ue_free(void* ptr) 
+{ 
+	FMemory::Free(ptr); 
+}
+
+static void* ue_realloc(void* ptr, size_t size) 
+{ 
+	return FMemory::Realloc(ptr, size); 
+}
+
+
+FUEJpegImageWrapper::FUEJpegImageWrapper(int32 InNumComponents /* = 4 */)
+	: FImageWrapperBase()
+	, NumComponents(InNumComponents)
+{ 
+	// NumComponents == 1 for GrayscaleJPEG
+	static uejpeg_alloc_t UeAlloc = { ue_malloc, ue_free, ue_realloc };
+	uejpeg_set_alloc(&UeAlloc);
+}
+
+
+FUEJpegImageWrapper::~FUEJpegImageWrapper()
+{
+}
+
+/* FImageWrapperBase interface
+ *****************************************************************************/
+
+bool FUEJpegImageWrapper::SetCompressed(const void* InCompressedData, int64 InCompressedSize)
+{
+	// jpgd doesn't support 64-bit sizes.
+	if (InCompressedSize < 0 || InCompressedSize > MAX_uint32)
+	{
+		return false;
+	}
+
+	int x, y, n;
+	unsigned char* d;
+
+	// Is it a JPEG?
+	if (!FMemory::Memcmp(InCompressedData, "\xFF\xD8\xFF", 3)) 
+	{
+		// First encode it to a UEJPEG format
+		unsigned char* UEJData = 0;
+		int UEJDataSize = 0;
+		{
+			uejpeg_encode_context_t ctx = uejpeg_encode_jpeg_mem_threaded_start((uint8*)InCompressedData, InCompressedSize, 0);	
+			if (ctx.error) 
+			{
+				return false;
+			}
+			ParallelFor(ctx.jobs_to_run, [&](int idx) {
+				uejpeg_encode_jpeg_thread_run(&ctx, idx);
+			});
+			int err = uejpeg_encode_jpeg_mem_threaded_finish(&ctx, &UEJData, &UEJDataSize);
+			if (err) 
+			{
+				return false;
+			}
+		}
+
+		// Then decode to a raw image
+		{
+			uejpeg_decode_context_t ctx = uejpeg_decode_mem_threaded_start(UEJData, UEJDataSize, UEJPEG_FLAG_FASTDCT);
+			if (ctx.error)
+			{
+				FMemory::Free(UEJData);
+				return false;
+			}
+			ParallelFor(ctx.jobs_to_run, [&](int idx) {
+				uejpeg_decode_thread_run(&ctx, idx);
+			});
+			d = uejpeg_decode_threaded_finish(&ctx, &x, &y, &n);
+		}
+		if (!d)
+		{
+			FMemory::Free(UEJData);
+			return false;
+		}
+		FMemory::Free(d);
+
+		if (!FImageWrapperBase::SetCompressed(UEJData, UEJDataSize))
+		{
+			FMemory::Free(UEJData);
+			return false;
+		}
+
+		FMemory::Free(UEJData);
+	}
+	else 
+	{
+		uejpeg_decode_context_t ctx = uejpeg_decode_mem_threaded_start((uint8*)InCompressedData, (uint32)InCompressedSize, UEJPEG_FLAG_FASTDCT);
+		if (ctx.error)
+		{
+			return false;
+		}
+		ParallelFor(ctx.jobs_to_run, [&](int idx) {
+			uejpeg_decode_thread_run(&ctx, idx);
+		});
+		d = uejpeg_decode_threaded_finish(&ctx, &x, &y, &n);
+		if (!d)
+		{
+			return false;
+		}
+		FMemory::Free(d);
+
+		if (!FImageWrapperBase::SetCompressed(InCompressedData, InCompressedSize))
+		{
+			return false;
+		}
+	}
+
+
+	// We don't support 16 bit jpegs
+	BitDepth = 8;
+
+	Width = x;
+	Height = y;
+
+	switch (n)
+	{
+	case 1:
+		Format = ERGBFormat::Gray;
+		break;
+	case 4:
+		Format = ERGBFormat::RGBA;
+		break;
+	default:
+		return false;
+	}
+	
+	if ( ! FImageCoreUtils::IsImageImportPossible(Width,Height) )
+	{
+		SetError(TEXT("Image dimensions are not possible to import"));
+		return false;
+	}
+
+	return true;
+}
+
+
+// CanSetRawFormat returns true if SetRaw will accept this format
+bool FUEJpegImageWrapper::CanSetRawFormat(const ERGBFormat InFormat, const int32 InBitDepth) const
+{
+	return ((InFormat == ERGBFormat::RGBA || InFormat == ERGBFormat::BGRA || InFormat == ERGBFormat::Gray) && InBitDepth == 8);
+}
+
+// returns InFormat if supported, else maps to something supported
+ERawImageFormat::Type FUEJpegImageWrapper::GetSupportedRawFormat(const ERawImageFormat::Type InFormat) const
+{
+	switch(InFormat)
+	{
+	case ERawImageFormat::G8:
+	case ERawImageFormat::BGRA8:
+		return InFormat; // directly supported
+	case ERawImageFormat::BGRE8:
+	case ERawImageFormat::RGBA16:
+	case ERawImageFormat::RGBA16F:
+	case ERawImageFormat::RGBA32F:
+	case ERawImageFormat::G16:
+	case ERawImageFormat::R16F:
+	case ERawImageFormat::R32F:
+		return ERawImageFormat::BGRA8; // needs conversion
+	default:
+		check(0);
+		return ERawImageFormat::BGRA8;
+	};
+}
+
+void FUEJpegImageWrapper::Compress(int32 Quality)
+{
+	if (Quality == 0)
+	{
+		//default 
+		Quality = 85;
+	}
+	else if (Quality == (int32)EImageCompressionQuality::Uncompressed)
+	{
+		// fix = 1 (Uncompressed) was previously treated as max-compress
+		Quality = 100;
+	}
+	else
+	{
+		ensure(Quality >= 1 && Quality <= 100);
+
+		#define JPEG_QUALITY_MIN 40
+		// JPEG should not be used below quality JPEG_QUALITY_MIN
+		Quality = FMath::Clamp(Quality, JPEG_QUALITY_MIN, 100);
+	}
+
+	if (CompressedData.Num() == 0)
+	{
+		check(RawData.Num());
+		check(Width > 0);
+		check(Height > 0);
+
+		// re-order components if required - JPEGs expect RGBA
+		if(Format == ERGBFormat::BGRA)
+		{
+			FColor* Colors = (FColor*)RawData.GetData();
+			const int32 NumColors = RawData.Num() / 4;
+			for(int32 ColorIndex = 0; ColorIndex < NumColors; ColorIndex++)
+			{
+				int32 tmp = Colors[ColorIndex].B;
+				Colors[ColorIndex].B = Colors[ColorIndex].R;
+				Colors[ColorIndex].R = tmp;
+			}
+		}
+
+
+		int data_size;
+		unsigned char* data;
+		uejpeg_encode_image_context_t ctx = uejpeg_encode_image_mem_threaded_start(Width, Height, NumComponents, RawData.GetData(), Quality, 0);
+		check(!ctx.error);
+		ParallelFor(ctx.jobs_to_run, [&](int idx) {
+			uejpeg_encode_image_thread_run(&ctx, idx);
+		});
+		int error = uejpeg_encode_image_mem_threaded_finish(&ctx, &data, &data_size);
+		check(!error);
+
+		CompressedData.Reset(data_size);
+		CompressedData.AddUninitialized(data_size);
+		FMemory::Memcpy(CompressedData.GetData(), data, data_size);
+		FMemory::Free(data);
+	}
+}
+
+
+void FUEJpegImageWrapper::Uncompress(const ERGBFormat InFormat, int32 InBitDepth)
+{
+	// Ensure we haven't already uncompressed the file.
+	if (RawData.Num() != 0)
+	{
+		return;
+	}
+
+	check(CompressedData.Num());
+
+	int32 NumColors;
+	int32 jpegWidth,jpegHeight;
+	uint8* OutData;
+	uejpeg_decode_context_t ctx = uejpeg_decode_mem_threaded_start(CompressedData.GetData(), CompressedData.Num(), UEJPEG_FLAG_FASTDCT);
+	if (ctx.error)
+	{
+		return;
+	}
+	ParallelFor(ctx.jobs_to_run, [&](int idx) {
+		uejpeg_decode_thread_run(&ctx, idx);
+	});
+	OutData = uejpeg_decode_threaded_finish(&ctx, &jpegWidth, &jpegHeight, &NumColors);
+
+	if (OutData)
+	{
+		Width = jpegWidth;
+		Height = jpegHeight;
+		RawData.Reset(Width * Height * NumColors);
+		RawData.AddUninitialized(Width * Height * NumColors);
+		FMemory::Memcpy(RawData.GetData(), OutData, RawData.Num());
+		FMemory::Free(OutData);
+	}
+	else
+	{
+		UE_LOG(LogImageWrapper, Error, TEXT("UEJPEG Decompress Error"));
+		RawData.Empty();
+	}
+}
+
+TArray64<uint8> FUEJpegImageWrapper::GetExportData(int32 Quality)
+{
+	uint8* Data;
+	int Size;
+	uejpeg_decode_mem_to_jpeg(CompressedData.GetData(), CompressedData.Num(), &Data, &Size, 0);
+
+	TArray64<uint8> ExportData;
+	ExportData.Reset();
+	ExportData.AddUninitialized(Size);
+	FMemory::Memcpy(ExportData.GetData(), Data, Size);
+	FMemory::Free(Data);
+	return ExportData;
+}
+
+#endif //WITH_UEJPEG
diff --git a/Engine/Source/Runtime/ImageWrapper/Private/Formats/UEJpegImageWrapper.h b/Engine/Source/Runtime/ImageWrapper/Private/Formats/UEJpegImageWrapper.h
new file mode 100644
--- /dev/null
+++ b/Engine/Source/Runtime/ImageWrapper/Private/Formats/UEJpegImageWrapper.h
@@ -0,0 +1,39 @@
+// Copyright Epic Games, Inc. All Rights Reserved.
+
+#pragma once
+
+#include "CoreTypes.h"
+#include "ImageWrapperBase.h"
+
+#if WITH_UEJPEG
+
+class FUEJpegImageWrapper
+	: public FImageWrapperBase
+{
+public:
+
+	/** Default constructor. */
+	// NumComponents == 1 for GrayscaleJPEG
+	FUEJpegImageWrapper(int32 InNumComponents = 4);
+
+	virtual ~FUEJpegImageWrapper();
+
+public:
+
+	//~ FImageWrapperBase interface
+
+	virtual bool SetCompressed(const void* InCompressedData, int64 InCompressedSize) override;
+	virtual void Uncompress(const ERGBFormat InFormat, int32 InBitDepth) override;
+	virtual void Compress(int32 Quality) override;
+	virtual TArray64<uint8> GetExportData(int32 Quality) override;
+	
+	virtual bool CanSetRawFormat(const ERGBFormat InFormat, const int32 InBitDepth) const override;
+	virtual ERawImageFormat::Type GetSupportedRawFormat(const ERawImageFormat::Type InFormat) const override;
+
+private:
+
+	int32 NumComponents;
+};
+
+
+#endif //WITH_UEJPEG
diff --git a/Engine/Source/Runtime/ImageWrapper/Private/ImageWrapperHelper.cpp b/Engine/Source/Runtime/ImageWrapper/Private/ImageWrapperHelper.cpp
--- a/Engine/Source/Runtime/ImageWrapper/Private/ImageWrapperHelper.cpp
+++ b/Engine/Source/Runtime/ImageWrapper/Private/ImageWrapperHelper.cpp
@@ -17,9 +17,9 @@ FStringView ImageWrapperHelper::GetFormatExtension(EImageFormat InImageFormat, b
 	case EImageFormat::GrayscaleJPEG:
 		StringExtension = TEXT(".jpg");
 		break;
-	case EImageFormat::OOJPEG:
-	case EImageFormat::GrayscaleOOJPEG:
-		StringExtension = TEXT(".ooj");
+	case EImageFormat::UEJPEG:
+	case EImageFormat::GrayscaleUEJPEG:
+		StringExtension = TEXT(".uej");
 		break;
 	case EImageFormat::BMP:
 		StringExtension = TEXT(".bmp");
@@ -77,9 +77,9 @@ EImageFormat ImageWrapperHelper::GetImageFormat(FStringView StringExtention)
 	{
 		return EImageFormat::JPEG;
 	}
-	if (StringExtention.Equals(TEXT("ooj"), ESearchCase::IgnoreCase))
+	if (StringExtention.Equals(TEXT("uej"), ESearchCase::IgnoreCase))
 	{
-		return EImageFormat::OOJPEG;
+		return EImageFormat::UEJPEG;
 	}
 	if (StringExtention.Equals(TEXT("bmp"), ESearchCase::IgnoreCase))
 	{
diff --git a/Engine/Source/Runtime/ImageWrapper/Private/ImageWrapperModule.cpp b/Engine/Source/Runtime/ImageWrapper/Private/ImageWrapperModule.cpp
--- a/Engine/Source/Runtime/ImageWrapper/Private/ImageWrapperModule.cpp
+++ b/Engine/Source/Runtime/ImageWrapper/Private/ImageWrapperModule.cpp
@@ -11,7 +11,7 @@
 #include "Formats/IcnsImageWrapper.h"
 #include "Formats/IcoImageWrapper.h"
 #include "Formats/JpegImageWrapper.h"
-#include "Formats/OOJpegImageWrapper.h"
+#include "Formats/UEJpegImageWrapper.h"
 #include "Formats/PngImageWrapper.h"
 #include "Formats/TgaImageWrapper.h"
 #include "Formats/TiffImageWrapper.h"
@@ -33,7 +33,7 @@ namespace
 	static const uint8 IMAGE_MAGIC_ICO[]  = {0x00, 0x00, 0x01, 0x00};
 	static const uint8 IMAGE_MAGIC_EXR[]  = {0x76, 0x2F, 0x31, 0x01};
 	static const uint8 IMAGE_MAGIC_ICNS[] = {0x69, 0x63, 0x6E, 0x73};
-	static const uint8 IMAGE_MAGIC_OOJPEG[] = {'O', 'O', 'J', 'P', 'E', 'G'};
+	static const uint8 IMAGE_MAGIC_UEJPEG[] = {'O', 'O', 'J', 'P', 'E', 'G'};
 
 	// Binary for #?RADIANCE
 	static const uint8 IMAGE_MAGIC_HDR[] = {0x23, 0x3f, 0x52, 0x41, 0x44, 0x49, 0x41, 0x4e, 0x43, 0x45, 0x0a};
@@ -101,15 +101,15 @@ public:
 			break;
 #endif	//WITH_UNREALJPEG
 
-#if WITH_OOJPEG
-		case EImageFormat::OOJPEG:
-			ImageWrapper = MakeShared<FOOJpegImageWrapper>();
+#if WITH_UEJPEG
+		case EImageFormat::UEJPEG:
+			ImageWrapper = MakeShared<FUEJpegImageWrapper>();
 			break;
 
-		case EImageFormat::GrayscaleOOJPEG:
-			ImageWrapper = MakeShared<FOOJpegImageWrapper>(1);
+		case EImageFormat::GrayscaleUEJPEG:
+			ImageWrapper = MakeShared<FUEJpegImageWrapper>(1);
 			break;
-#endif // WITH_OOJPEG
+#endif // WITH_UEJPEG
 
 		case EImageFormat::BMP:
 			ImageWrapper = MakeShared<FBmpImageWrapper>();
@@ -168,9 +168,9 @@ public:
 		{
 			Format = EImageFormat::JPEG; // @Todo: Should we detect grayscale vs non-grayscale?
 		}
-		else if (StartsWith((uint8*)CompressedData, CompressedSize, IMAGE_MAGIC_OOJPEG))
+		else if (StartsWith((uint8*)CompressedData, CompressedSize, IMAGE_MAGIC_UEJPEG))
 		{
-			Format = EImageFormat::OOJPEG; // @Todo: Should we detect grayscale vs non-grayscale?
+			Format = EImageFormat::UEJPEG; // @Todo: Should we detect grayscale vs non-grayscale?
 		}
 		else if (StartsWith((uint8*) CompressedData, CompressedSize, IMAGE_MAGIC_BMP))
 		{
@@ -223,8 +223,8 @@ public:
 		case EImageFormat::PNG: return TEXT("png");
 		case EImageFormat::JPEG: return TEXT("jpg");
 		case EImageFormat::GrayscaleJPEG: return TEXT("jpg");
-		case EImageFormat::OOJPEG: return TEXT("ooj");
-		case EImageFormat::GrayscaleOOJPEG: return TEXT("ooj");
+		case EImageFormat::UEJPEG: return TEXT("uej");
+		case EImageFormat::GrayscaleUEJPEG: return TEXT("uej");
 		case EImageFormat::BMP: return TEXT("bmp");
 		case EImageFormat::ICO: return TEXT("ico");
 		case EImageFormat::EXR: return TEXT("exr");
@@ -256,9 +256,9 @@ public:
 		{
 			return EImageFormat::JPEG;
 		}
-		else if (FCString::Stricmp(Name, TEXT("ooj")) == 0)
+		else if (FCString::Stricmp(Name, TEXT("uej")) == 0)
 		{
-			return EImageFormat::OOJPEG;
+			return EImageFormat::UEJPEG;
 		}
 		else if ( FCString::Stricmp(Name,TEXT("bmp")) == 0 )
 		{
diff --git a/Engine/Source/Runtime/ImageWrapper/Public/IImageWrapper.h b/Engine/Source/Runtime/ImageWrapper/Public/IImageWrapper.h
--- a/Engine/Source/Runtime/ImageWrapper/Public/IImageWrapper.h
+++ b/Engine/Source/Runtime/ImageWrapper/Public/IImageWrapper.h
@@ -60,11 +60,11 @@ enum class EImageFormat : int8
 	/** DirectDraw Surface */
 	DDS,
 
-	/** Oodle JPEG format. */
-	OOJPEG,
+	/** UE JPEG format. */
+	UEJPEG,
 
-	/** Single channel Oodle JPEG. */
-	GrayscaleOOJPEG,
+	/** Single channel UE JPEG. */
+	GrayscaleUEJPEG,
 };
 
 
diff --git a/Engine/Source/Runtime/OodleJPEG/OodleJPEG.Build.cs b/Engine/Source/Runtime/OodleJPEG/OodleJPEG.Build.cs
deleted file mode 100644
diff --git a/Engine/Source/Runtime/OodleJPEG/Source/oojpeg.h b/Engine/Source/Runtime/OodleJPEG/Source/oojpeg.h
deleted file mode 100644
diff --git a/Engine/Source/Runtime/UEJpegComp/Source/uejpeg.h b/Engine/Source/Runtime/UEJpegComp/Source/uejpeg.h
new file mode 100644
--- /dev/null
+++ b/Engine/Source/Runtime/UEJpegComp/Source/uejpeg.h
@@ -0,0 +1,283 @@
+// Copyright Epic Games, Inc. All Rights Reserved.
+/*
+
+UEJPEG - A JPEG Super-compressor using an oodle backend as the entropy coder.
+
+Can encode from a JPEG file to an UEJPEG file, and can decode directly to a JPEG file or a raw image.
+Can also encode raw images directly to UEJPEG files.
+
+API has simple APIs and threaded APIs. 
+
+The simple APIs is as follows:
+-----------------------------
+
+The following functions encode from a JPEG file to a UEJPEG file...
+
+  int uejpeg_encode_jpeg_mem(const unsigned char *jpeg_data, int jpeg_data_size, unsigned char **out_uejpeg_data, int *out_uejpeg_size, int flags);
+
+The following functions encode from a raw RGBA image to a UEJPEG file...
+
+  int uejpeg_encode_to_mem(unsigned char **out_data, int *out_size, int width, int height, int channels, const unsigned char *rgba, int quality, int flags);
+
+The following functions decode from a UEJPEG file back to a JPEG file...
+
+  int uejpeg_decode_mem_to_jpeg(const unsigned char *data, int size, unsigned char **out_data, int *out_size, int flags);
+
+The following functions decode from a UEJPEG file to a raw RGBA image...
+
+  unsigned char *uejpeg_decode_mem(const unsigned char *data, int size, int *width, int *height, int *channels, int flags);
+
+The threaded APIs is as follows:
+-------------------------------
+
+The following functions encode from a JPEG file to a UEJPEG file...
+
+  uejpeg_encode_context_t uejpeg_encode_jpeg_mem_threaded_start(const unsigned char *data, int size, int flags);
+  int uejpeg_encode_jpeg_thread_run(uejpeg_encode_context_t *context, int job_idx);
+  int uejpeg_encode_jpeg_mem_threaded_finish(uejpeg_encode_context_t *context, unsigned char **out_data, int *out_size);
+
+The following functions encode from a raw RGBA image to a UEJPEG file...
+
+  uejpeg_encode_image_context_t uejpeg_encode_image_mem_threaded_start(int width, int height, int comp, const unsigned char *in_data, int quality, int flags);
+  uejpeg_encode_image_context_t uejpeg_encode_image_threaded_start(const uejpeg_io_callbacks_t *io, void *io_user, int width, int height, int comp, const unsigned char *in_data, int quality, int flags);
+  int uejpeg_encode_image_thread_run(uejpeg_encode_image_context_t *ctx, int job_idx);
+  int uejpeg_encode_image_threaded_finish(uejpeg_encode_image_context_t *ctx);
+  int uejpeg_encode_image_mem_threaded_finish(uejpeg_encode_image_context_t *ctx, unsigned char **out_data, int *out_size);
+
+The following functions decode from a UEJPEG file to a raw RGBA image...
+
+  uejpeg_decode_context_t uejpeg_decode_mem_threaded_start(const unsigned char *data, int size, int flags);
+  uejpeg_decode_context_t uejpeg_decode_threaded_start(const uejpeg_io_callbacks_t *io, void *io_user, int flags);
+  int uejpeg_decode_thread_run(uejpeg_decode_context_t *context, int job_idx);
+  unsigned char *uejpeg_decode_threaded_finish(uejpeg_decode_context_t *context, int *out_width, int *out_height, int *out_comp);
+
+You use the threaded APIs by first calling the _start function, which gives you a context. You then pass that context into the thread_run function N times,
+where N is the value of the context.jobs_to_run variable. You then wait for all the threads to finish, and then call the _finish function to get the output data.
+
+Example:
+
+      uejpeg_encode_context_t ctx = uejpeg_encode_jpeg_mem_threaded_start(input, output, flags);
+      #pragma omp parallel for
+      for(int i = 0; i < ctx.jobs_to_run; ++i) {
+        uejpeg_encode_jpeg_thread_run(&ctx, i);
+      }
+      uejpeg_encode_jpeg_file_threaded_finish(&ctx);
+
+Other Notes:
+
+* Custom Allocators can be set via uejpeg_set_alloc() function.
+* Custom compression/decompression (must be set the same in both the encoder and decoder) can be set via uejpeg_set_compression() function.
+
+*/
+#pragma once
+
+#ifdef __cplusplus
+extern "C" {
+#endif
+
+// Error codes
+typedef enum {
+  UEJPEG_ERR_NONE                = 0, // No Error, everything OK! :)
+  UEJPEG_ERR_NOT_A_JPEG          = 1, // Not a JPEG file. Happens if you try to encode from a non-JPEG file.
+  UEJPEG_ERR_CANNOT_OPEN_FILE    = 2, // Cannot open file. Happens if the file routines cannot open the files you specified!
+  UEJPEG_ERR_NOT_UEJPEG_FILE     = 3, // Not a uejpeg file. Happens if you try to decode a non-uejpeg file.
+  UEJPEG_ERR_CORRUPT_UEJPEG_FILE = 4, // Corrupt uejpeg file. Happens if the uejpeg file is corrupt.
+  UEJPEG_ERR_WRITE_FAILED        = 5, // Write failed. Happens if the file routines cannot write to the files you specified!
+  UEJPEG_ERR_OUT_OF_MEMORY       = 6, // Out of memory. Happens if the uejpeg library cannot allocate enough memory.
+  UEJPEG_ERR_UNSUPPORTED_COLORSPACE = 7, // Unsupported colorspace. Happens when exporting if the JPEG file is not YCbCr.
+  UEJPEG_ERR_READ_FAILED         = 8, // Read failed. Happens if the file routines cannot read from the files you specified!
+  UEJPEG_ERR_UNSUPPORTED         = 9, // Other Unsupported. Happens usually when you try to decode to jpeg in an uncommon format.
+} uejpeg_error_t;
+
+// Flags to be used in compression or decompression
+typedef enum {
+  UEJPEG_FLAG_NONE            = 0x0, // no flags
+  UEJPEG_FLAG_NO_COMPRESSION  = 0x1, // don't compress the coefficients, just store them
+  UEJPEG_FLAG_FASTDCT         = 0x2, // use a fast dct transform which matches libjpegturbo's fastdct.
+} uejpeg_flags_t;
+
+enum {
+  UEJPEG_MAX_SPLITS = 16,
+};
+
+// IO callbacks
+//  For use with uejpeg_*() to specify any input or output method. All methods internally use this!
+typedef struct {
+  size_t (*write)(void *user, const void *data, size_t size);  // write 'data' with 'size' bytes.  return number of bytes actually written
+  size_t (*read)(void *user, void *data, size_t size);  // read 'data' with 'size' bytes.  return number of bytes actually read
+  int (*eof)(void *user);  // return 1 if end-of-file, 0 otherwise
+} uejpeg_io_callbacks_t;
+
+// For using memory buffers as input or output.
+typedef struct {
+  char *base;
+  char *ptr;
+  size_t size;
+} uejpeg_io_buffer_t;
+
+extern const uejpeg_io_callbacks_t s_uejpeg_buffer_fns;
+
+// For setting custom memory allocators.
+typedef struct {
+  void *(*alloc)(size_t size);  // allocate 'size' bytes.  return pointer to allocated memory
+  void (*free)(void *ptr);  // free memory at 'ptr'
+  void *(*realloc)(void *ptr, size_t size);  // reallocate 'ptr' to 'size' bytes.  return pointer to reallocated memory
+} uejpeg_alloc_t;
+
+typedef struct {
+    void *data;
+    size_t size;
+} uejpeg_buffer_t;
+
+// compression callbacks, for using custom compression methods, such as zlib
+typedef struct {
+  uejpeg_buffer_t (*compress)(const void *data, size_t size);
+  int (*decompress)(const void *data, size_t size, void *out_data, size_t out_size);
+} uejpeg_compression_t;
+
+// Set custom allocators.  If not set, malloc/free/realloc will be used.
+void uejpeg_set_alloc(const uejpeg_alloc_t *alloc);
+// Set custom compression methods.  If not set, oodle compression methods will be used.
+void uejpeg_set_compression(const uejpeg_compression_t *compression);
+
+// An internal structure used to store the header of a uejpeg file.
+typedef struct uejpeg_ihdr_t {
+  unsigned version;
+  unsigned width;
+  unsigned height;
+  unsigned char bit_depth;
+  unsigned char comp;
+  unsigned char method;
+  unsigned char num_splits;
+} uejpeg_ihdr_t;
+
+// An internal structure used to store the header of a uejpeg file.
+typedef struct uejpeg_lossy_hdr_t {
+  unsigned short coef_dims[4][2];
+  // Quantization tables
+  unsigned short dequant[4][64];
+  char app14_color_transform; // valid values are 0(CMYK),1(YCCK),2(YCbCrA)
+  char jfif;
+  char comp_id[4];
+} uejpeg_lossy_hdr_t;
+
+typedef struct {
+  int y_start[4];
+  int y_end[4];
+  int width[4];
+  int height[4];
+  int num_blocks;
+  short *xcoefs;
+  char *split_xcoefs;
+  unsigned char *oodle_buf;
+  int oodle_size;
+  int coefs_size;
+} uejpeg_internal_split_t;
+
+typedef struct {
+  int error;
+  int split_height_min;
+  int split_height[4];
+  int split_blocks[4];
+  int total_split_blocks;
+  uejpeg_internal_split_t splits[UEJPEG_MAX_SPLITS];
+} uejpeg_internal_splits_t;
+
+typedef struct {
+  int error; // If there was an error, it is stored here.
+  int jobs_to_run; // How many times to run uejpeg_decode_thread_run!
+
+  // arguments
+  const uejpeg_io_callbacks_t *io;
+  void *io_user;
+  int flags;
+
+  // internal state
+  int width, height, comp;
+  int num_splits;
+
+  uejpeg_lossy_hdr_t lhdr;
+  uejpeg_internal_splits_t s;
+} uejpeg_encode_context_t;
+
+typedef struct {
+  int error; // If there was an error, it is stored here.
+  int jobs_to_run; // How many times to run uejpeg_decode_thread_run!
+
+  // arguments
+  int flags;
+
+  // internal state
+  uejpeg_ihdr_t ihdr;
+  uejpeg_lossy_hdr_t lhdr;
+  uejpeg_internal_splits_t spx;
+  int yuva_stride[4];
+  int yuva_height[4];
+  unsigned char *yuva[4];
+} uejpeg_decode_context_t;
+
+typedef struct {
+  int error; // If there was an error, it is stored here.
+  int jobs_to_run; // How many times to run uejpeg_decode_thread_run!
+
+  // arguments
+  const uejpeg_io_callbacks_t *io;
+  void *io_user;
+  int width, height, comp;
+  const unsigned char *in_data;
+  int flags;
+
+  // internal state
+  int num_splits;
+  uejpeg_internal_splits_t spx;
+  uejpeg_lossy_hdr_t lhdr;
+	float fdtbl_Y[64];
+  float fdtbl_UV[64];
+} uejpeg_encode_image_context_t;
+
+// Encode a JPEG file to uejpeg format.
+// This second function operates totally in memory.
+int uejpeg_encode_jpeg_mem(const unsigned char *jpeg_data, int jpeg_data_size, unsigned char **out_uejpeg_data, int *out_uejpeg_size, int flags);
+
+// Threaded API
+// In this API you first call start, and it does some synchronous work and gives you a state context to pass to the run function.
+// Then inside threads, you call the run function with job_idx == 0 .. context.jobs_to_run
+// Then you wait for all the threads to finish processing all the jobs.
+// Then you call finish to get the output data.
+uejpeg_encode_context_t uejpeg_encode_jpeg_mem_threaded_start(const unsigned char *data, int size, int flags);
+int uejpeg_encode_jpeg_thread_run(uejpeg_encode_context_t *context, int job_idx);
+int uejpeg_encode_jpeg_mem_threaded_finish(uejpeg_encode_context_t *context, unsigned char **out_data, int *out_size);
+
+// Encode a raw image to a uejpeg file.
+int uejpeg_encode_to_mem(unsigned char **out_data, int *out_size, int width, int height, int channels, const unsigned char *rgba, int quality, int flags);
+
+uejpeg_encode_image_context_t uejpeg_encode_image_mem_threaded_start(int width, int height, int comp, const unsigned char *in_data, int quality, int flags);
+uejpeg_encode_image_context_t uejpeg_encode_image_threaded_start(const uejpeg_io_callbacks_t *io, void *io_user, int width, int height, int comp, const unsigned char *in_data, int quality, int flags);
+int uejpeg_encode_image_thread_run(uejpeg_encode_image_context_t *ctx, int job_idx);
+int uejpeg_encode_image_threaded_finish(uejpeg_encode_image_context_t *ctx);
+int uejpeg_encode_image_mem_threaded_finish(uejpeg_encode_image_context_t *ctx, unsigned char **out_data, int *out_size);
+
+// Decode a UEJPEG file to a raw image.
+// Interface is like stbi_load_from_file, stbi_load_from_memory, stbi_load, etc.
+unsigned char *uejpeg_decode_mem(const unsigned char *data, int size, int *width, int *height, int *channels, int flags);
+unsigned char *uejpeg_decode(const uejpeg_io_callbacks_t *io, void *io_user, int *out_width, int *out_height, int *out_comp, int flags);
+
+// Threaded API
+// In this API you first call start, and it does some synchronous work and gives you a state context to pass to the run function.
+// Then inside threads, you call the run function with job_idx == 0 .. context.jobs_to_run
+// Then you wait for all the threads to finish processing all the jobs.
+// Then you call finish to get the output data.
+
+uejpeg_decode_context_t uejpeg_decode_mem_threaded_start(const unsigned char *data, int size, int flags);
+uejpeg_decode_context_t uejpeg_decode_threaded_start(const uejpeg_io_callbacks_t *io, void *io_user, int flags);
+int uejpeg_decode_thread_run(uejpeg_decode_context_t *context, int job_idx);
+unsigned char *uejpeg_decode_threaded_finish(uejpeg_decode_context_t *context, int *out_width, int *out_height, int *out_comp);
+
+// Decode a UEJPEG file to a JPEG file.
+int uejpeg_decode_to_jpeg(const uejpeg_io_callbacks_t *in_io, void *in_io_user, const uejpeg_io_callbacks_t *out_io, void *out_io_user, int flags);
+int uejpeg_decode_mem_to_jpeg(const unsigned char *data, int size, unsigned char **out_data, int *out_size, int flags);
+
+#ifdef __cplusplus
+}
+#endif
+
diff --git a/Engine/Source/Runtime/UEJpegComp/UEJpegComp.Build.cs b/Engine/Source/Runtime/UEJpegComp/UEJpegComp.Build.cs
new file mode 100644
--- /dev/null
+++ b/Engine/Source/Runtime/UEJpegComp/UEJpegComp.Build.cs
@@ -0,0 +1,38 @@
+// Copyright Epic Games, Inc. All Rights Reserved.
+using UnrealBuildTool;
+using System.IO;
+
+public class UEJpegComp : ModuleRules
+{
+	public UEJpegComp(ReadOnlyTargetRules Target) : base(Target)
+	{
+		Type = ModuleType.External;
+
+		bool bWithLib = false;
+		if (Target.Platform == UnrealTargetPlatform.Win64 && Target.WindowsPlatform.Architecture != UnrealArch.Arm64)
+		{
+			string LibPath = Path.Combine(ModuleDirectory, "Lib");
+			PublicAdditionalLibraries.Add(Path.Combine(LibPath, "uejpeg_w64.lib"));
+			bWithLib = true;
+		}
+		else if (Target.Platform == UnrealTargetPlatform.Mac)
+		{
+			string LibPath = Path.Combine(ModuleDirectory, "Lib");
+			PublicAdditionalLibraries.Add(Path.Combine(LibPath, "uejpeg_mac.a"));
+			bWithLib = true;
+		}
+		else if (Target.IsInPlatformGroup(UnrealPlatformGroup.Unix))
+		{
+			string LibPath = Path.Combine(ModuleDirectory, "Lib/Unix", Target.Architecture.LinuxName);
+			PublicAdditionalLibraries.Add(Path.Combine(LibPath, "uejpeg_linux.a"));
+			bWithLib = true;
+		}
+
+		if (bWithLib)
+		{
+			PublicSystemIncludePaths.Add(Path.Combine(ModuleDirectory, "Source"));
+		}
+
+		PublicDefinitions.Add("WITH_UEJPEG=" + (bWithLib ? '1' : '0'));
+	}
+}
