arm64 support for uejpeg
- also added missing files to source zip

#rnx

diff --git a/Engine/Source/Runtime/UEJpegComp/UEJpegComp.Build.cs b/Engine/Source/Runtime/UEJpegComp/UEJpegComp.Build.cs
--- a/Engine/Source/Runtime/UEJpegComp/UEJpegComp.Build.cs
+++ b/Engine/Source/Runtime/UEJpegComp/UEJpegComp.Build.cs
@@ -15,6 +15,12 @@ public class UEJpegComp : ModuleRules
 			PublicAdditionalLibraries.Add(Path.Combine(LibPath, "uejpeg_w64.lib"));
 			bWithLib = true;
 		}
+		if (Target.Platform == UnrealTargetPlatform.Win64 && Target.WindowsPlatform.Architecture == UnrealArch.Arm64)
+		{
+			string LibPath = Path.Combine(ModuleDirectory, "Lib");
+			PublicAdditionalLibraries.Add(Path.Combine(LibPath, "uejpeg_warm64.lib"));
+			bWithLib = true;
+		}
 		else if (Target.Platform == UnrealTargetPlatform.Mac)
 		{
 			string LibPath = Path.Combine(ModuleDirectory, "Lib");
